[Mutable] Robustify clothing vertex remove mask generation agains invalid data.

- Add LogMutableClothing log category.
- Add logging when the vertex remove mask generation finds invalid data.

#rnx

diff --git a/Engine/Plugins/Experimental/MutableClothing/Source/MutableClothing/Private/MutableClothingModule.cpp b/Engine/Plugins/Experimental/MutableClothing/Source/MutableClothing/Private/MutableClothingModule.cpp
--- a/Engine/Plugins/Experimental/MutableClothing/Source/MutableClothing/Private/MutableClothingModule.cpp
+++ b/Engine/Plugins/Experimental/MutableClothing/Source/MutableClothing/Private/MutableClothingModule.cpp
@@ -6,8 +6,9 @@
 #include "Modules/ModuleManager.h"
 #include "ClothingAsset.h"
 
-IMPLEMENT_MODULE(FMutableClothingModule, MutableClothing);
+IMPLEMENT_MODULE(FMutableClothingModule, MutableClothing)
 
+DEFINE_LOG_CATEGORY(LogMutableClothing)
 
 void FMutableClothingModule::StartupModule()
 {
@@ -38,18 +39,42 @@ bool FMutableClothingModule::UpdateClothSimulationLOD(
 	TBitArray<> VertexRemoveMask;
 	VertexRemoveMask.Init(true, NumSimulationVertices);
 
+	int32 NumInvalidIndicesFound = 0;
 	for (int32 RenderDataIndex = 0; RenderDataIndex < InOutAttachedLODsRenderData.Num(); ++RenderDataIndex)
 	{
 		for (const FMeshToMeshVertData& VertClothData : InOutAttachedLODsRenderData[RenderDataIndex])
 		{
 			const uint16* Indices = VertClothData.SourceMeshVertIndices;
 
-			VertexRemoveMask[Indices[0]] = false;
-			VertexRemoveMask[Indices[1]] = false;
-			VertexRemoveMask[Indices[2]] = false;
+			if (Indices[0] < NumSimulationVertices)
+			{
+				VertexRemoveMask[Indices[0]] = false;
+			}
+
+			if (Indices[1] < NumSimulationVertices)
+			{
+				VertexRemoveMask[Indices[1]] = false;
+			}
+
+			if (Indices[2] < NumSimulationVertices)
+			{
+				VertexRemoveMask[Indices[2]] = false;
+			}
+
+#if DO_CHECK
+			NumInvalidIndicesFound += 
+					static_cast<int32>(Indices[0] >= NumSimulationVertices) +
+					static_cast<int32>(Indices[1] >= NumSimulationVertices) +
+					static_cast<int32>(Indices[2] >= NumSimulationVertices);
+#endif
 		}
 	}	
 
+	if (NumInvalidIndicesFound > 0)
+	{
+		UE_LOG(LogMutableClothing, Error, TEXT("Invalid clothing render data indices found."));
+	}
+
 	const int32 FirstRemoved = VertexRemoveMask.Find(true);
 	if (FirstRemoved == INDEX_NONE)
 	{
diff --git a/Engine/Plugins/Experimental/MutableClothing/Source/MutableClothing/Private/MutableClothingModule.h b/Engine/Plugins/Experimental/MutableClothing/Source/MutableClothing/Private/MutableClothingModule.h
--- a/Engine/Plugins/Experimental/MutableClothing/Source/MutableClothing/Private/MutableClothingModule.h
+++ b/Engine/Plugins/Experimental/MutableClothing/Source/MutableClothing/Private/MutableClothingModule.h
@@ -3,8 +3,11 @@
 #pragma once
 
 #include "MuCO/Plugins/IMutableClothingModule.h"
+#include "Logging/LogMacros.h"
 
 
+MUTABLECLOTHING_API DECLARE_LOG_CATEGORY_EXTERN(LogMutableClothing, Log, All);
+
 class MUTABLECLOTHING_API FMutableClothingModule : public IMutableClothingModule
 {
 public:
